Fix clang compile error in MCPClientToolset.

#Synced-CL 51932664

diff --git a/Engine/Plugins/Experimental/MCPClientToolset/Source/MCPClientToolset/Private/MCPClientToolset/MCPClientToolset.cpp b/Engine/Plugins/Experimental/MCPClientToolset/Source/MCPClientToolset/Private/MCPClientToolset/MCPClientToolset.cpp
--- a/Engine/Plugins/Experimental/MCPClientToolset/Source/MCPClientToolset/Private/MCPClientToolset/MCPClientToolset.cpp
+++ b/Engine/Plugins/Experimental/MCPClientToolset/Source/MCPClientToolset/Private/MCPClientToolset/MCPClientToolset.cpp
@@ -1544,7 +1544,7 @@ bool FMCPClientToolset::HandleOAuthCallback(
 	if (!CodePtr || CodePtr->IsEmpty() || !StatePtr || StatePtr->IsEmpty())
 	{
 		TUniquePtr<FHttpServerResponse> ErrorResp =
-			FHttpServerResponse::Create(TEXT("Missing code or state parameter"), TEXT("text/plain"));
+			FHttpServerResponse::Create(FString(TEXT("Missing code or state parameter")), TEXT("text/plain"));
 		ErrorResp->Code = EHttpServerResponseCodes::BadRequest;
 		OnComplete(MoveTemp(ErrorResp));
 		FailOAuth(FString::Printf(TEXT("OAuth callback missing parameters for '%s'"), *Config.Name));
@@ -1554,7 +1554,7 @@ bool FMCPClientToolset::HandleOAuthCallback(
 	if (*StatePtr != OAuthCSRFState)
 	{
 		TUniquePtr<FHttpServerResponse> ErrorResp =
-			FHttpServerResponse::Create(TEXT("OAuth state mismatch"), TEXT("text/plain"));
+			FHttpServerResponse::Create(FString(TEXT("OAuth state mismatch")), TEXT("text/plain"));
 		ErrorResp->Code = EHttpServerResponseCodes::BadRequest;
 		OnComplete(MoveTemp(ErrorResp));
 		FailOAuth(FString::Printf(TEXT("OAuth state mismatch for '%s'"), *Config.Name));
@@ -1563,7 +1563,7 @@ bool FMCPClientToolset::HandleOAuthCallback(
 
 	// Send the success page before cleaning up
 	TUniquePtr<FHttpServerResponse> SuccessResp = FHttpServerResponse::Create(
-		TEXT("<html><body><h1>Authorization complete</h1><p>You may close this tab.</p></body></html>"),
+		FString(TEXT("<html><body><h1>Authorization complete</h1><p>You may close this tab.</p></body></html>")),
 		TEXT("text/html"));
 	OnComplete(MoveTemp(SuccessResp));
 
