[MCP Client Toolset] Only allow the plugin to be packaged with editor targets.

#Synced-CL 51985403

diff --git a/Engine/Plugins/Experimental/MCPClientToolset/MCPClientToolset.uplugin b/Engine/Plugins/Experimental/MCPClientToolset/MCPClientToolset.uplugin
--- a/Engine/Plugins/Experimental/MCPClientToolset/MCPClientToolset.uplugin
+++ b/Engine/Plugins/Experimental/MCPClientToolset/MCPClientToolset.uplugin
@@ -25,7 +25,8 @@
 		{
 			"Name": "MCPClientToolset",
 			"Type": "Editor",
-			"LoadingPhase": "Default"
+			"LoadingPhase": "Default",
+			"TargetAllowList": [ "Editor" ]
 		}
 	]
 }
