[MCP Client Toolset] Fix unneccessary export of UMCPToolsetSettings class.

The constructor can be exported by MinimalAPI and all methods are inline.

#Synced-CL 51990377

diff --git a/Engine/Plugins/Experimental/MCPClientToolset/Source/MCPClientToolset/Public/MCPClientToolset/MCPToolsetSettings.h b/Engine/Plugins/Experimental/MCPClientToolset/Source/MCPClientToolset/Public/MCPClientToolset/MCPToolsetSettings.h
--- a/Engine/Plugins/Experimental/MCPClientToolset/Source/MCPClientToolset/Public/MCPClientToolset/MCPToolsetSettings.h
+++ b/Engine/Plugins/Experimental/MCPClientToolset/Source/MCPClientToolset/Public/MCPClientToolset/MCPToolsetSettings.h
@@ -70,8 +70,8 @@ struct FMCPServerConfig
 
 /** Editor settings for MCP (Model Context Protocol) server connections.
  *  Configure servers under Editor Preferences > Plugins > MCP Toolset Servers. */
-UCLASS(Config = EditorPerProjectUserSettings, DefaultConfig, meta = (DisplayName = "MCP Toolset Servers"))
-class MCPCLIENTTOOLSET_API UMCPToolsetSettings : public UDeveloperSettings
+UCLASS(Config = EditorPerProjectUserSettings, DefaultConfig, meta = (DisplayName = "MCP Toolset Servers"), MinimalAPI)
+class UMCPToolsetSettings : public UDeveloperSettings
 {
 	GENERATED_BODY()
 
