[MCP client toolset] Do not enable by default.

#Synced-CL 51993721

diff --git a/Engine/Plugins/Experimental/MCPClientToolset/MCPClientToolset.uplugin b/Engine/Plugins/Experimental/MCPClientToolset/MCPClientToolset.uplugin
--- a/Engine/Plugins/Experimental/MCPClientToolset/MCPClientToolset.uplugin
+++ b/Engine/Plugins/Experimental/MCPClientToolset/MCPClientToolset.uplugin
@@ -10,7 +10,7 @@
 	"DocsURL": "",
 	"MarketplaceURL": "",
 	"SupportURL": "",
-	"EnabledByDefault": true,
+	"EnabledByDefault": false,
 	"CanContainContent": false,
 	"IsBetaVersion": true,
 	"IsExperimentalVersion": true,
