[MCP client toolset] Fix compiler warnings

WeakThis shadows a member variable of TSharedFromThis

#Synced-CL 51993721

diff --git a/Engine/Plugins/Experimental/MCPClientToolset/Source/MCPClientToolset/Private/MCPClientToolset/MCPClientToolset.cpp b/Engine/Plugins/Experimental/MCPClientToolset/Source/MCPClientToolset/Private/MCPClientToolset/MCPClientToolset.cpp
--- a/Engine/Plugins/Experimental/MCPClientToolset/Source/MCPClientToolset/Private/MCPClientToolset/MCPClientToolset.cpp
+++ b/Engine/Plugins/Experimental/MCPClientToolset/Source/MCPClientToolset/Private/MCPClientToolset/MCPClientToolset.cpp
@@ -222,10 +222,8 @@ TFuture<TValueOrError<FString, FString>> FMCPClientToolset::ExecuteTool(
 	Params->SetStringField(TEXT("name"), ToolName);
 	Params->SetObjectField(TEXT("arguments"), InputJson);
 
-	TWeakPtr<FMCPClientToolset> WeakThis = AsWeak();
-
 	return SendRequest(TEXT("tools/call"), Params)
-		.Next([WeakThis](TSharedPtr<FJsonObject> Response) -> TValueOrError<FString, FString>
+		.Next([WeakThis = AsWeak()](TSharedPtr<FJsonObject> Response) -> TValueOrError<FString, FString>
 		{
 			TSharedPtr<FMCPClientToolset> StrongThis = WeakThis.Pin();
 			if (!StrongThis.IsValid())
@@ -328,14 +326,12 @@ void FMCPClientToolset::StartSSEConnection()
 			*Config.Name);
 	}
 
-	TWeakPtr<FMCPClientToolset> WeakThis = AsWeak();
-
 	// Drain the SPSC queue on the game thread each frame so SSE events are dispatched there.
 	SSETickerHandle = FTSTicker::GetCoreTicker().AddTicker(
 		FTickerDelegate::CreateSP(AsShared(), &FMCPClientToolset::OnSSETick));
 
 	SSERequest->OnProcessRequestComplete().BindLambda(
-		[WeakThis](FHttpRequestPtr Request, FHttpResponsePtr Response, bool bConnectedSuccessfully)
+		[WeakThis = AsWeak()](FHttpRequestPtr Request, FHttpResponsePtr Response, bool bConnectedSuccessfully)
 		{
 			if (TSharedPtr<FMCPClientToolset> StrongThis = WeakThis.Pin())
 			{
@@ -526,10 +522,8 @@ void FMCPClientToolset::DispatchSSEEvent(const FString& EventType, const FString
 		ClientInfo->SetStringField(TEXT("version"), TEXT("1.0"));
 		InitParams->SetObjectField(TEXT("clientInfo"), ClientInfo);
 
-		TWeakPtr<FMCPClientToolset> WeakThis = AsWeak();
-
 		SendRequest(TEXT("initialize"), InitParams)
-			.Next([WeakThis](TSharedPtr<FJsonObject> InitResponse)
+			.Next([WeakThis = AsWeak()](TSharedPtr<FJsonObject> InitResponse)
 			{
 				TSharedPtr<FMCPClientToolset> StrongThis = WeakThis.Pin();
 				if (!StrongThis.IsValid())
@@ -740,11 +734,10 @@ TFuture<TSharedPtr<FJsonObject>> FMCPClientToolset::SendRequest(
 
 	// If the HTTP POST itself fails (network error, non-2xx), reject the promise immediately.
 	// On success the actual JSON-RPC response arrives via SSE, not via the POST response body.
-	TWeakPtr<FMCPClientToolset> WeakThis = AsWeak();
 	const int32 CapturedId = RequestId;
 
 	HttpRequest->OnProcessRequestComplete().BindLambda(
-		[WeakThis, CapturedId](FHttpRequestPtr /*Req*/, FHttpResponsePtr Response, bool bSuccess)
+		[WeakThis = AsWeak(), CapturedId](FHttpRequestPtr /*Req*/, FHttpResponsePtr Response, bool bSuccess)
 		{
 			// A 202 Accepted means the server queued the request; response via SSE.
 			// Anything else (or a network failure) is a hard error.
@@ -912,10 +905,8 @@ void FMCPClientToolset::StartStreamableHTTPHandshake()
 	ClientInfo->SetStringField(TEXT("version"), TEXT("1.0"));
 	InitParams->SetObjectField(TEXT("clientInfo"), ClientInfo);
 
-	TWeakPtr<FMCPClientToolset> WeakThis = AsWeak();
-
 	SendStreamableHTTPRequest(TEXT("initialize"), InitParams)
-		.Next([WeakThis](TSharedPtr<FJsonObject> InitResponse)
+		.Next([WeakThis = AsWeak()](TSharedPtr<FJsonObject> InitResponse)
 		{
 			TSharedPtr<FMCPClientToolset> StrongThis = WeakThis.Pin();
 			if (!StrongThis.IsValid())
@@ -1093,11 +1084,10 @@ TFuture<TSharedPtr<FJsonObject>> FMCPClientToolset::SendStreamableHTTPRequest(
 				}
 			}));
 
-	TWeakPtr<FMCPClientToolset> WeakThis = AsWeak();
 	const int32 CapturedId = RequestId;
 
 	HttpRequest->OnProcessRequestComplete().BindLambda(
-		[WeakThis, CapturedStreamState, CapturedId]
+		[WeakThis = AsWeak(), CapturedStreamState, CapturedId]
 		(FHttpRequestPtr /*Req*/, FHttpResponsePtr Response, bool bSuccess)
 		{
 			TSharedPtr<FMCPClientToolset> StrongThis = WeakThis.Pin();
@@ -1285,9 +1275,8 @@ void FMCPClientToolset::StartOAuthFlow()
 	Request->SetURL(DiscoveryUrl);
 	Request->SetVerb(TEXT("GET"));
 
-	TWeakPtr<FMCPClientToolset> WeakThis = AsWeak();
 	Request->OnProcessRequestComplete().BindLambda(
-		[WeakThis](FHttpRequestPtr Req, FHttpResponsePtr Response, bool bSuccess)
+		[WeakThis = AsWeak()](FHttpRequestPtr Req, FHttpResponsePtr Response, bool bSuccess)
 		{
 			if (TSharedPtr<FMCPClientToolset> StrongThis = WeakThis.Pin())
 			{
@@ -1409,9 +1398,8 @@ void FMCPClientToolset::RegisterOAuthClient()
 	Request->SetHeader(TEXT("Content-Type"), TEXT("application/json"));
 	Request->SetContentAsString(RegBodyStr);
 
-	TWeakPtr<FMCPClientToolset> WeakThis = AsWeak();
 	Request->OnProcessRequestComplete().BindLambda(
-		[WeakThis](FHttpRequestPtr Req, FHttpResponsePtr Response, bool bSuccess)
+		[WeakThis = AsWeak()](FHttpRequestPtr Req, FHttpResponsePtr Response, bool bSuccess)
 		{
 			if (TSharedPtr<FMCPClientToolset> StrongThis = WeakThis.Pin())
 			{
@@ -1493,12 +1481,11 @@ void FMCPClientToolset::LaunchOAuthBrowser()
 
 	// 5. Bind GET /callback and start the listener
 	// (OAuthLocalPort and OAuthRouter were set in OnOAuthDiscoveryComplete)
-	TWeakPtr<FMCPClientToolset> WeakThis = AsWeak();
 	OAuthRouteHandle = OAuthRouter->BindRoute(
 		FHttpPath(TEXT("/callback")),
 		EHttpServerRequestVerbs::VERB_GET,
 		FHttpRequestHandler::CreateLambda(
-			[WeakThis](const FHttpServerRequest& Request, const FHttpResultCallback& OnComplete) -> bool
+			[WeakThis = AsWeak()](const FHttpServerRequest& Request, const FHttpResultCallback& OnComplete) -> bool
 			{
 				if (TSharedPtr<FMCPClientToolset> StrongThis = WeakThis.Pin())
 				{
@@ -1604,9 +1591,8 @@ void FMCPClientToolset::ExchangeOAuthCode(const FString& Code)
 	Request->SetHeader(TEXT("Content-Type"), TEXT("application/x-www-form-urlencoded"));
 	Request->SetContentAsString(PostBody);
 
-	TWeakPtr<FMCPClientToolset> WeakThis = AsWeak();
 	Request->OnProcessRequestComplete().BindLambda(
-		[WeakThis](FHttpRequestPtr Req, FHttpResponsePtr Response, bool bSuccess)
+		[WeakThis = AsWeak()](FHttpRequestPtr Req, FHttpResponsePtr Response, bool bSuccess)
 		{
 			if (TSharedPtr<FMCPClientToolset> StrongThis = WeakThis.Pin())
 			{
@@ -1778,9 +1764,8 @@ void FMCPClientToolset::TryRefreshToken()
 	Request->SetHeader(TEXT("Content-Type"), TEXT("application/x-www-form-urlencoded"));
 	Request->SetContentAsString(Body);
 
-	TWeakPtr<FMCPClientToolset> WeakThis = AsWeak();
 	Request->OnProcessRequestComplete().BindLambda(
-		[WeakThis](FHttpRequestPtr Req, FHttpResponsePtr Response, bool bSuccess)
+		[WeakThis = AsWeak()](FHttpRequestPtr Req, FHttpResponsePtr Response, bool bSuccess)
 		{
 			if (TSharedPtr<FMCPClientToolset> StrongThis = WeakThis.Pin())
 			{
