Change MCPToolsetSettings config to Editor so that it's set at a per-user level

- Users were editing this property and inadvertantly checking out DefaultPerProjectUserSettings.ini
- If checked in they could accidentally change default MCP Server for all devs

#Synced-CL 52444574

#Synced-CL 52444574

diff --git a/Engine/Plugins/Experimental/Toolsets/MCPClientToolset/Source/MCPClientToolset/Public/MCPClientToolset/MCPToolsetSettings.h b/Engine/Plugins/Experimental/Toolsets/MCPClientToolset/Source/MCPClientToolset/Public/MCPClientToolset/MCPToolsetSettings.h
--- a/Engine/Plugins/Experimental/Toolsets/MCPClientToolset/Source/MCPClientToolset/Public/MCPClientToolset/MCPToolsetSettings.h
+++ b/Engine/Plugins/Experimental/Toolsets/MCPClientToolset/Source/MCPClientToolset/Public/MCPClientToolset/MCPToolsetSettings.h
@@ -70,7 +70,7 @@ struct FMCPServerConfig
 
 /** Editor settings for MCP (Model Context Protocol) server connections.
  *  Configure servers under Editor Preferences > Plugins > MCP Toolset Servers. */
-UCLASS(Config = EditorPerProjectUserSettings, DefaultConfig, meta = (DisplayName = "MCP Toolset Servers"), MinimalAPI)
+UCLASS(Config = Editor, meta = (DisplayName = "MCP Toolset Servers"), MinimalAPI)
 class UMCPToolsetSettings : public UDeveloperSettings
 {
 	GENERATED_BODY()
