MCPToolsetSettings - change back to Config=DefaultPerProjectUserSettings to fix serialization issues

#Synced-CL 52539024

diff --git a/Engine/Plugins/Experimental/Toolsets/MCPClientToolset/Source/MCPClientToolset/Public/MCPClientToolset/MCPToolsetSettings.h b/Engine/Plugins/Experimental/Toolsets/MCPClientToolset/Source/MCPClientToolset/Public/MCPClientToolset/MCPToolsetSettings.h
--- a/Engine/Plugins/Experimental/Toolsets/MCPClientToolset/Source/MCPClientToolset/Public/MCPClientToolset/MCPToolsetSettings.h
+++ b/Engine/Plugins/Experimental/Toolsets/MCPClientToolset/Source/MCPClientToolset/Public/MCPClientToolset/MCPToolsetSettings.h
@@ -70,7 +70,7 @@ struct FMCPServerConfig
 
 /** Editor settings for MCP (Model Context Protocol) server connections.
  *  Configure servers under Editor Preferences > Plugins > MCP Toolset Servers. */
-UCLASS(Config = Editor, meta = (DisplayName = "MCP Toolset Servers"), MinimalAPI)
+UCLASS(Config = EditorPerProjectUserSettings, meta = (DisplayName = "MCP Toolset Servers"), MinimalAPI)
 class UMCPToolsetSettings : public UDeveloperSettings
 {
 	GENERATED_BODY()
